Separating the OnlineTests into Null and EOS executables: Creating the OnlineServicesNullTests project.

#Synced-CL 48759110

diff --git a/Engine/Build/LowLevelTests/OnlineServicesNull.xml b/Engine/Build/LowLevelTests/OnlineServicesNull.xml
new file mode 100644
--- /dev/null
+++ b/Engine/Build/LowLevelTests/OnlineServicesNull.xml
@@ -0,0 +1,16 @@
+<?xml version="1.0" encoding="utf-8"?>
+<BuildGraph xmlns:xsi="http://www.w3.org/2001/XMLSchema-instance" xsi:schemaLocation="http://www.epicgames.com/BuildGraph ../../Build/Graph/Schema.xsd" xmlns="http://www.epicgames.com/BuildGraph">
+  <Macro Name="OnlineServicesNullAfterSteps" Arguments="Platform">
+    <Spawn Exe="$(RootDir)/Engine/Restricted/NotForLicensees/Source/Programs/OnlineTests/Utilities/dist/CatchLogAggregatorNew.exe" Arguments="--jobid $(UE_HORDE_JOBID) --platform $(Platform) --testname OnlineServicesNull" />
+  </Macro>
+  <Option Name="RunOnlineServicesNullTests" DefaultValue="" Description="Run Online Services Null Tests" />
+  <Property Name="TestNames" Value="$(TestNames);OnlineServicesNull" />
+  <Extend Name="RunAllTests">
+    <Expand Name="DeployAndTest" Platform="Win64" TestName="OnlineServicesNull" ShortName="OnlineServicesNull" StagesWithProjectFile="True" TargetName="OnlineServicesNullTests" BinaryRelativePath="Engine\Binaries" ReportType="xml" ExtraArgs="-dumpiniloads" Tags="~[StringProcessing]" />
+    <Expand Name="DeployAndTest" Platform="Linux" TestName="OnlineServicesNull" ShortName="OnlineServicesNull" StagesWithProjectFile="True" TargetName="OnlineServicesNullTests" BinaryRelativePath="Engine\Binaries" ReportType="xml" ExtraArgs="-dumpiniloads" Tags="~[StringProcessing]" />
+    <Expand Name="DeployAndTest" Platform="Mac" TestName="OnlineServicesNull" ShortName="OnlineServicesNull" StagesWithProjectFile="True" TargetName="OnlineServicesNullTests" BinaryRelativePath="Engine\Binaries" ReportType="xml" ExtraArgs="-dumpiniloads" Tags="~[StringProcessing]" />
+    <!-- Currently needs test framework support, leave out from running on Horde for now. -->
+    <!-- <Expand Name="DeployAndTest" Platform="IOS" TestName="OnlineServicesNull" ShortName="OnlineServicesNull" StagesWithProjectFile="True" TargetName="OnlineServicesNullTests" BinaryRelativePath="Engine\Binaries" ReportType="xml" PlatformGauntletArgs="~[StringProcessing]" ExtraArgs="-dumpiniloads" Tags="~[StringProcessing]" RunUnsupported="True" /> -->
+    <!-- <Expand Name="DeployAndTest" Platform="Android" TestName="OnlineServicesNull" ShortName="OnlineServicesNull" StagesWithProjectFile="True" TargetName="OnlineServicesNullTests" BinaryRelativePath="Engine\Binaries" ReportType="xml" PlatformGauntletArgs="~[StringProcessing]" ExtraArgs="-dumpiniloads" Tags="~[StringProcessing]" RunUnsupported="True" /> -->
+  </Extend>
+</BuildGraph>
\ No newline at end of file
diff --git a/Engine/Programs/Online/OnlineServicesNullTests/Config/DefaultEngine.ini b/Engine/Programs/Online/OnlineServicesNullTests/Config/DefaultEngine.ini
new file mode 100644
--- /dev/null
+++ b/Engine/Programs/Online/OnlineServicesNullTests/Config/DefaultEngine.ini
@@ -0,0 +1,116 @@
+[Core.Log]
+; This can be used to change the default log level for engine logs to help with debugging
+LogTemp=VeryVerbose
+LogHTTP=Verbose
+LogOnlineIdentity=Verbose
+LogOnlineServices=Verbose
+;LogHandshake=VeryVerbose
+
+[OnlineServicesTests]
++DefaultServices="Null"
+
+;Null test configuration
+[OnlineServicesTests Null]
+ServicesType="Null"
++ModulesToLoad="OnlineServicesNull"
+
+[LoginCredentials Null]
+UseImplicitLogin=true
+
+[TestReporting Null]
++DisableTestTags=Lobby
++DisableTestTags=suite_social
++DisableTestTags=Meta
++DisableTestTags=suite_commerce
++DisableTestTags=Leaderboards
++DisableTestTags=userfile
+
+;Steam via OSSAdapter test configuration
+[OnlineServicesTests Steam]
+ServicesType="Steam"
++ModulesToLoad="OnlineSubsystem"
++ModulesToLoad="OnlineSubsystemSteam"
++ModulesToLoad="OnlineServicesOSSAdapter"
+
+[OnlineServices.OSSAdapter]
++Services=(Service="Steam",ConfigName="OSSAdapter.Steam",OnlineSubsystem="Steam")
+
+[LoginCredentials Steam]
+UseAutoLogin=true
+
+[OnlineServices]
+bUseBuildIdOverride=false
+BuildIdOverride=
+
+[OnlineServices.Null.Auth]
+bAddUserNumToNullId=true
+
+; TODO: See if we can put following configs into editor then export to platform portals and .ini config here
+
+[OnlineServices.Null.Achievements]
++AchievementDefinitions=(AchievementId="test_unlockachievements", UnlockedDisplayName="test_unlockachievements Unlocked Displayname", UnlockedDescription="Unlocked description of test_unlockachievements", LockedDisplayName="test_unlockachievements Locked Displayname", LockedDescription="Locked description of test_unlockachievements", FlavorText="test_unlockachievements flavor", UnlockedIconUrl="http://example.com/test_unlockachievementsunlocked.jpg", LockedIconUrl="http://example.com/test_unlockachievementslocked.jpg", bIsHidden=false, StatDefinitions=((StatId="StatA", UnlockThreshold=100), (StatId="StatB", UnlockThreshold=200)))
++AchievementDefinitions=(AchievementId="FAKE_ACHIEVEMENT", UnlockedDisplayName="FAKE_ACHIEVEMENT Unlocked Displayname", UnlockedDescription="Unlocked description of FAKE_ACHIEVEMENT", LockedDisplayName="FAKE_ACHIEVEMENT Locked Displayname", LockedDescription="Locked description of Achi1", FlavorText="FAKE_ACHIEVEMENT flavor", UnlockedIconUrl="http://example.com/FAKE_ACHIEVEMENTunlocked.jpg", LockedIconUrl="http://example.com/FAKE_ACHIEVEMENTlocked.jpg", bIsHidden=true)
+
+[OnlineServices.Achievements]
+bIsTitleManaged=true
++UnlockRules=(AchievementId=test_unlockachievements, Conditions=((StatName=Stat_Use_Largest, UnlockThreshold="Int64:3"), (StatName=Stat_Use_Smallest, UnlockThreshold="Int64:1")))
++UnlockRules=(AchievementId=FAKE_ACHIEVEMENT, Conditions=((StatName=Stat_Use_Set, UnlockThreshold="Int64:1")))
+ 
+[OnlineServices.Null.TitleFile]
++Files=(Name=FileA, Contents="The quick brown fox jumps over the lazy dog")
++Files=(Name=FileB, Contents="Lorem ipsum dolor sit amet")
+
+[OnlineServices.Stats]
++StatDefinitions=(Name=Stat_Use_Largest, Id=0, ModifyMethod=Largest, DefaultValue="Int64:0")
++StatDefinitions=(Name=Stat_Use_Smallest, Id=1, ModifyMethod=Smallest, DefaultValue="Int64:999")
++StatDefinitions=(Name=Stat_Use_Set, Id=2, ModifyMethod=Set, DefaultValue="Int64:0")
++StatDefinitions=(Name=Stat_Use_Sum, Id=3, ModifyMethod=Sum, DefaultValue="Int64:0")
++StatDefinitions=(Name=Stat_Type_Bool, Id=4, ModifyMethod=Set, DefaultValue="Bool:True")
++StatDefinitions=(Name=Stat_Type_Double, Id=5, ModifyMethod=Smallest, DefaultValue="Double:9999.999")
+
+[OnlineServices.Leaderboards]
+bIsTitleManaged=true
++LeaderboardDefinitions=(Name=Stat_Use_Smallest, Id=0, UpdateMethod=KeepBest, OrderMethod=Ascending)
++LeaderboardDefinitions=(Name=Stat_Use_Set, Id=1, UpdateMethod=Force, OrderMethod=Descending)
+
+[OnlineServices.Lobbies]
+!SchemaCategoryAttributeDescriptors=ClearArray
++SchemaCategoryAttributeDescriptors=(SchemaId="LobbyBase", CategoryId="Lobby", AttributeIds=("SchemaCompatibilityId", "LobbyCreateTime"))
++SchemaCategoryAttributeDescriptors=(SchemaId="LobbyBase", CategoryId="LobbyMember")
++SchemaAttributeDescriptors=(Id="LobbyCreateTime", Type="Int64", Flags=("Public", "Searchable"))
++SchemaDescriptors=(Id="Test", ParentId="LobbyBase")
+
+[OnlineServices.NULL.Lobbies]
+!ServiceDescriptors=ClearArray
++ServiceDescriptors=(Id="Lobby", AttributeIds=("LobbyServiceAttribute1", "LobbyServiceAttribute2", "LobbyServiceAttribute3", "LobbyServiceAttribute4", "LobbyServiceAttribute5"))
++ServiceDescriptors=(Id="LobbyMember", AttributeIds=("LobbyMemberServiceAttribute1", "LobbyMemberServiceAttribute2", "LobbyMemberServiceAttribute3", "LobbyMemberServiceAttribute4", "LobbyMemberServiceAttribute5"))
+!ServiceAttributeDescriptors=ClearArray
++ServiceAttributeDescriptors=(Id="LobbyServiceAttribute1", SupportedTypes=("Bool", "Int64", "Double", "String"), Flags=("Searchable", "Public", "Private"), MaxSize=1000)
++ServiceAttributeDescriptors=(Id="LobbyServiceAttribute2", SupportedTypes=("Bool", "Int64", "Double", "String"), Flags=("Searchable", "Public", "Private"), MaxSize=1000)
++ServiceAttributeDescriptors=(Id="LobbyServiceAttribute3", SupportedTypes=("Bool", "Int64", "Double", "String"), Flags=("Searchable", "Public", "Private"), MaxSize=1000)
++ServiceAttributeDescriptors=(Id="LobbyServiceAttribute4", SupportedTypes=("Bool", "Int64", "Double", "String"), Flags=("Searchable", "Public", "Private"), MaxSize=1000)
++ServiceAttributeDescriptors=(Id="LobbyServiceAttribute5", SupportedTypes=("Bool", "Int64", "Double", "String"), Flags=("Searchable", "Public", "Private"), MaxSize=1000)
++ServiceAttributeDescriptors=(Id="LobbyMemberServiceAttribute1", SupportedTypes=("Bool", "Int64", "Double", "String"), Flags=("Searchable", "Public", "Private"), MaxSize=1000)
++ServiceAttributeDescriptors=(Id="LobbyMemberServiceAttribute2", SupportedTypes=("Bool", "Int64", "Double", "String"), Flags=("Searchable", "Public", "Private"), MaxSize=1000)
++ServiceAttributeDescriptors=(Id="LobbyMemberServiceAttribute3", SupportedTypes=("Bool", "Int64", "Double", "String"), Flags=("Searchable", "Public", "Private"), MaxSize=1000)
++ServiceAttributeDescriptors=(Id="LobbyMemberServiceAttribute4", SupportedTypes=("Bool", "Int64", "Double", "String"), Flags=("Searchable", "Public", "Private"), MaxSize=1000)
++ServiceAttributeDescriptors=(Id="LobbyMemberServiceAttribute5", SupportedTypes=("Bool", "Int64", "Double", "String"), Flags=("Searchable", "Public", "Private"), MaxSize=1000)
+
+[OnlineTests.TestCacheDuration]
+CacheExpiration=Duration
+CacheExpirySeconds=2
+bCacheError=false
+
+[OnlineTests.TestDontCacheError]
+CacheExpiration=Duration
+CacheExpirySeconds=1
+bCacheError=false
+
+[OnlineTests.TestCacheError]
+CacheExpiration=Duration
+CacheExpirySeconds=1
+bCacheError=true
+
+[OnlineTests.TestCacheExpireUponCompletion]
+CacheExpiration=UponCompletion
+bCacheError=false
\ No newline at end of file
diff --git a/Engine/Programs/Online/OnlineServicesNullTests/Config/DefaultGame.ini b/Engine/Programs/Online/OnlineServicesNullTests/Config/DefaultGame.ini
new file mode 100644
--- /dev/null
+++ b/Engine/Programs/Online/OnlineServicesNullTests/Config/DefaultGame.ini
@@ -0,0 +1,12 @@
+[Staging]
+;Remap platform folders for OnlineTests loading configs from these specific folders for some reason only on console.
++RemapDirectories=(From="Engine/Platforms/PS4/Programs/OnlineServicesNullTests",To="OnlineServicesNullTests/Platforms/PS4")
++RemapDirectories=(From="Engine/Binaries/PS5/OnlineServicesNullTests",To="OnlineServicesNullTests/Binaries/PS5")
++RemapDirectories=(From="Engine/Platforms/PS5/Programs/OnlineServicesNullTests",To="OnlineServicesNullTests/Platforms/PS5")
++RemapDirectories=(From="Engine/Platforms/Sony/Programs/OnlineServicesNullTests",To="OnlineServicesNullTests/Platforms/Sony")
++RemapDirectories=(From="Engine/Platforms/Switch/Programs/OnlineServicesNullTests",To="OnlineServicesNullTests/Platforms/Switch")
++RemapDirectories=(From="Engine/Platforms/XSX/Programs/OnlineServicesNullTests",To="OnlineServicesNullTests/Platforms/XSX")
++RemapDirectories=(From="Engine/Binaries/Linux/OnlineServicesNullTests",To="Engine/Binaries/Linux")
+
+; Can't move binary location for BuildCookRun deploy won't work.
++AllowedDirectories=Engine/Restricted/NotForLicensees
\ No newline at end of file
diff --git a/Engine/Programs/Online/OnlineServicesNullTests/OnlineServicesNullTests.uproject b/Engine/Programs/Online/OnlineServicesNullTests/OnlineServicesNullTests.uproject
new file mode 100644
--- /dev/null
+++ b/Engine/Programs/Online/OnlineServicesNullTests/OnlineServicesNullTests.uproject
@@ -0,0 +1,11 @@
+{
+	"FileVersion": 3,
+	"EngineAssociation": "",
+	"Category": "",
+	"Description": "",
+	"TargetPlatforms": [
+		"Windows",
+		"Mac",
+		"Linux"
+	]
+}
\ No newline at end of file
diff --git a/Engine/Source/Programs/Online/OnlineServicesNullTests/OnlineServicesNullTests.Build.cs b/Engine/Source/Programs/Online/OnlineServicesNullTests/OnlineServicesNullTests.Build.cs
new file mode 100644
--- /dev/null
+++ b/Engine/Source/Programs/Online/OnlineServicesNullTests/OnlineServicesNullTests.Build.cs
@@ -0,0 +1,46 @@
+// Copyright Epic Games, Inc. All Rights Reserved.
+
+using UnrealBuildTool;
+using System.IO;
+
+public class OnlineServicesNullTests : TestModuleRules
+{
+	static OnlineServicesNullTests()
+	{
+		TestMetadata = new Metadata();
+		TestMetadata.TestName = "OnlineServicesNullTests";
+		TestMetadata.TestShortName = "OnlineServicesNull";
+		TestMetadata.ReportType = "xml";
+		TestMetadata.StagesWithProjectFile = true;
+
+		// Ensure LowLevelTestsRunner adds Engine/CoreUObject/ApplicationCore to its PrivateDependencyModuleNames.
+		// LowLevelTestsRunner checks these TestTargetRules static flags when configuring its module.
+		TestTargetRules.bTestsRequireEngine = true;
+		TestTargetRules.bTestsRequireCoreUObject = true;
+		TestTargetRules.bTestsRequireApplicationCore = true;
+	}
+
+	public OnlineServicesNullTests(ReadOnlyTargetRules Target) : base(Target)
+	{
+		bTreatAsEngineModule = false;
+		CppStandard = CppStandardVersion.EngineDefault;
+		IncludeOrderVersion = EngineIncludeOrderVersion.Latest;
+
+		PrivateDependencyModuleNames.AddRange(
+			new string[] {
+				"ApplicationCore",
+				"Core",
+				"CoreUObject",
+				"Engine",
+				"OnlineSubsystem",
+				"OnlineServicesInterface",
+				"OnlineServicesNull",
+				"OnlineTestsCore",
+				"Json",
+				"JsonUtilities"
+			});
+
+		PublicDefinitions.Add("ONLINETESTS_REQUIREAPPLICATIONTICK=0");
+		PublicDefinitions.Add("ONLINETESTS_REQUIREPLATFORMINIT=0");
+	}
+}
\ No newline at end of file
diff --git a/Engine/Source/Programs/Online/OnlineServicesNullTests/OnlineServicesNullTests.Target.cs b/Engine/Source/Programs/Online/OnlineServicesNullTests/OnlineServicesNullTests.Target.cs
new file mode 100644
--- /dev/null
+++ b/Engine/Source/Programs/Online/OnlineServicesNullTests/OnlineServicesNullTests.Target.cs
@@ -0,0 +1,39 @@
+// Copyright Epic Games, Inc. All Rights Reserved.
+
+using System;
+using System.IO;
+using UnrealBuildTool;
+using EpicGames.Core;
+using UnrealBuildBase;
+
+[SupportedPlatforms(UnrealPlatformClass.All)]
+public class OnlineServicesNullTestsTarget : TestTargetRules
+{
+	[CommandLine("-NoPreBuildTargets", Value = "false")]
+	protected bool bUsePreBuildTargets = true;
+	protected virtual string OnlineServicesNullTestsProjectName { get { return "OnlineServicesNullTests"; } }
+	public OnlineServicesNullTestsTarget(TargetInfo Target) : base(Target)
+	{
+		Type = TargetType.Program;
+
+		bCompileAgainstEngine = true; // Although the LLT framework finds the Engine module in the dependency graph, it currently is optimized to never compile against the engine for explicit tests like this one
+		bCompileAgainstCoreUObject = true;
+		bCompileAgainstApplicationCore = true;
+
+		if (Target.Platform == UnrealTargetPlatform.Mac)
+		{
+			bIsBuildingConsoleApplication = false;
+		}
+
+		bForceEnableExceptions = true; // Catch testing framework uses exceptions to abort from individual tests
+
+		bMockEngineDefaults = true;
+		bUsePlatformFileStub = true;
+
+		// ICU data files don't need to be copied and deployed on Windows, enable it and do Unicode string tests shouldn't add much overhead
+		if (Target.Platform == UnrealTargetPlatform.Win64)
+		{
+			bCompileICU = false;
+		}
+	}
+}
\ No newline at end of file
diff --git a/Engine/Source/Programs/Online/OnlineServicesNullTests/Private/ExampleTest.cpp b/Engine/Source/Programs/Online/OnlineServicesNullTests/Private/ExampleTest.cpp
new file mode 100644
--- /dev/null
+++ b/Engine/Source/Programs/Online/OnlineServicesNullTests/Private/ExampleTest.cpp
@@ -0,0 +1,7 @@
+// Copyright Epic Games, Inc. All Rights Reserved.
+#include "OnlineCatchHelper.h"
+
+TEST_CASE("OnlineServiceNull Basic Test", "[OnlineServiceNull]")
+{
+	REQUIRE(true);
+}
\ No newline at end of file
diff --git a/Engine/Source/Programs/Online/OnlineServicesNullTests/Private/OnlineLLTModule.cpp b/Engine/Source/Programs/Online/OnlineServicesNullTests/Private/OnlineLLTModule.cpp
new file mode 100644
--- /dev/null
+++ b/Engine/Source/Programs/Online/OnlineServicesNullTests/Private/OnlineLLTModule.cpp
@@ -0,0 +1,86 @@
+// Copyright Epic Games, Inc. All Rights Reserved.
+
+#include "GenericPlatform/GenericPlatformMisc.h"
+#include "Online/OnlineExecHandler.h"
+#include "Misc/CoreDelegates.h"
+#include "Misc/CommandLine.h"
+#include "Misc/ConfigCacheIni.h"
+#include "TestCommon/Initialization.h"
+#include "TestCommon/CoreUObjectUtilities.h"
+#include "TestRunner.h"
+
+#include "OnlineCatchHelper.h"
+
+#if ONLINETESTS_REQUIREAPPLICATIONTICK
+#include "HAL/PlatformApplicationMisc.h"
+#include "Containers/Ticker.h"
+#endif
+
+#include <catch2/catch_test_macros.hpp>
+
+#if ONLINETESTS_REQUIREAPPLICATIONTICK
+static GenericApplication* GPlatformApplication = nullptr;
+static FTSTicker::FDelegateHandle GPlatformApplicationTickHandle;
+#endif
+
+void InitializeOnlineTests()
+{
+#if ONLINETESTS_REQUIREPLATFORMINIT
+	extern void PlatformInitializeTests();
+	PlatformInitializeTests();
+#endif
+	SetProjectNameAndDirectory();
+	InitAll(true, true);
+
+	UseDefaultPlatformFile();
+	FConfigCacheIni::InitializeConfigSystem();
+	CleanupCoreUObject();
+
+#if ONLINETESTS_REQUIREAPPLICATIONTICK
+	GPlatformApplication = FPlatformApplicationMisc::CreateApplication();
+	check(GPlatformApplication);
+	GPlatformApplicationTickHandle = FTSTicker::GetCoreTicker().AddTicker(TEXT("PlatformApplicationTick"), 0.0f,
+		[](float DeltaSeconds)
+		{
+			GPlatformApplication->Tick(DeltaSeconds);
+			return true;
+		});
+#endif
+
+	FLogSuppressionInterface::Get().ProcessConfigAndCommandLine();
+	OnlineTestBase::LoadServiceModules();
+
+	auto GlobalInitalizersPtr = GetGlobalInitalizers();
+	ensure(GlobalInitalizersPtr && GlobalInitalizersPtr->Num() > 0);
+	if (GlobalInitalizersPtr && GlobalInitalizersPtr->Num() > 0)
+	{
+		for (const auto& Init : *GlobalInitalizersPtr)
+		{
+			Init();
+		}		
+	}	
+
+	// Some Service Modules rely on this broadcast to finish a late initialization. 
+	FCoreDelegates::OnFEngineLoopInitComplete.Broadcast();
+}
+
+void CleanupOnlineTests()
+{
+#if ONLINETESTS_REQUIREAPPLICATIONTICK
+	FTSTicker::RemoveTicker(GPlatformApplicationTickHandle);
+	GPlatformApplicationTickHandle.Reset();
+#endif
+
+	CleanupAll();
+	OnlineTestBase::UnloadServiceModules();
+}
+
+class OnlineTestGlobalSetup
+{
+public:
+	OnlineTestGlobalSetup()
+	{
+		FTestDelegates::GetGlobalSetup().BindStatic(&InitializeOnlineTests);
+		FTestDelegates::GetGlobalTeardown().BindStatic(&CleanupOnlineTests);
+	}
+} GOnlineTestGlobalSetup;
diff --git a/Engine/Source/Programs/Online/OnlineTestsCore/Source/OnlineTestsCore.Build.cs b/Engine/Source/Programs/Online/OnlineTestsCore/Source/OnlineTestsCore.Build.cs
--- a/Engine/Source/Programs/Online/OnlineTestsCore/Source/OnlineTestsCore.Build.cs
+++ b/Engine/Source/Programs/Online/OnlineTestsCore/Source/OnlineTestsCore.Build.cs
@@ -39,12 +39,6 @@ public class OnlineTestsCore : ModuleRules
 		{
 			PublicDefinitions.Add(String.Format("ONLINETESTS_USEEXTERNAUTH=0"));
 		}
-
-		PrivateIncludePaths.AddRange(
-			new string[] {
-				Path.Combine(Target.RelativeEnginePath, "Restricted", "NotForLicensees", "Source", "Programs", "OnlineTests", "Private")
-			}
-		);
 	}
 }
 
